Renaming template TP_nDisplay to TP_nDisplayBP

diff --git a/Templates/TP_nDisplay/Config/DefaultEditor.ini b/Templates/TP_nDisplay/Config/DefaultEditor.ini
deleted file mode 100644
diff --git a/Templates/TP_nDisplay/Config/DefaultEngine.ini b/Templates/TP_nDisplay/Config/DefaultEngine.ini
deleted file mode 100644
diff --git a/Templates/TP_nDisplay/Config/DefaultGame.ini b/Templates/TP_nDisplay/Config/DefaultGame.ini
deleted file mode 100644
diff --git a/Templates/TP_nDisplay/Config/DefaultInput.ini b/Templates/TP_nDisplay/Config/DefaultInput.ini
deleted file mode 100644
diff --git a/Templates/TP_nDisplay/Config/TemplateDefs.ini b/Templates/TP_nDisplay/Config/TemplateDefs.ini
deleted file mode 100644
diff --git a/Templates/TP_nDisplay/TP_nDisplay.uproject b/Templates/TP_nDisplay/TP_nDisplay.uproject
deleted file mode 100644
diff --git a/Templates/TP_nDisplayBP/Config/DefaultEngine.ini b/Templates/TP_nDisplayBP/Config/DefaultEngine.ini
new file mode 100644
--- /dev/null
+++ b/Templates/TP_nDisplayBP/Config/DefaultEngine.ini
@@ -0,0 +1,80 @@
+[/Script/HardwareTargeting.HardwareTargetingSettings]
+TargetedHardwareClass=Desktop
+AppliedTargetedHardwareClass=Desktop
+DefaultGraphicsPerformance=Maximum
+AppliedDefaultGraphicsPerformance=Maximum
+
+[/Script/EngineSettings.GameMapsSettings]
+EditorStartupMap=/Game/Maps/Main.Main
+GameDefaultMap=/Game/Maps/Main.Main
+GlobalDefaultGameMode=/Script/DisplayCluster.DisplayClusterGameModeDefault
+
+[/Script/Engine.Engine]
+GameEngine=/Script/DisplayCluster.DisplayClusterGameEngine
+UnrealEdEngine=/Script/DisplayClusterEditor.DisplayClusterEditorEngine
+
+[/Script/DisplayClusterEditor.DisplayClusterEditorSettings]
+bEnabled=True
+
+[/Script/Engine.PhysicsSettings]
+DefaultGravityZ=-980.000000
+DefaultTerminalVelocity=4000.000000
+DefaultFluidFriction=0.300000
+SimulateScratchMemorySize=262144
+RagdollAggregateThreshold=4
+TriangleMeshTriangleMinAreaThreshold=5.000000
+bEnableAsyncScene=False
+bEnableShapeSharing=False
+bEnablePCM=True
+bEnableStabilization=False
+bWarnMissingLocks=True
+bEnable2DPhysics=False
+PhysicErrorCorrection=(LinearDeltaThresholdSq=5.000000,LinearInterpAlpha=0.200000,LinearRecipFixTime=1.000000,AngularDeltaThreshold=0.628319,AngularInterpAlpha=0.100000,AngularRecipFixTime=1.000000,BodySpeedThresholdSq=0.200000)
+LockedAxis=Invalid
+DefaultDegreesOfFreedom=Full3D
+BounceThresholdVelocity=200.000000
+FrictionCombineMode=Average
+RestitutionCombineMode=Average
+MaxAngularVelocity=3600.000000
+MaxDepenetrationVelocity=0.000000
+ContactOffsetMultiplier=0.020000
+MinContactOffset=2.000000
+MaxContactOffset=8.000000
+bSimulateSkeletalMeshOnDedicatedServer=True
+DefaultShapeComplexity=CTF_UseSimpleAndComplex
+bDefaultHasComplexCollision=True
+bSuppressFaceRemapTable=False
+bSupportUVFromHitResults=False
+bDisableActiveActors=False
+bDisableCCD=False
+bEnableEnhancedDeterminism=False
+MaxPhysicsDeltaTime=0.033333
+bSubstepping=False
+bSubsteppingAsync=False
+MaxSubstepDeltaTime=0.016667
+MaxSubsteps=6
+SyncSceneSmoothingFactor=0.000000
+AsyncSceneSmoothingFactor=0.990000
+InitialAverageFrameRate=0.016667
+PhysXTreeRebuildRate=10
+
+[/Script/WindowsTargetPlatform.WindowsTargetSettings]
+Compiler=Default
+-TargetedRHIs=PCD3D_SM5
+-TargetedRHIs=PCD3D_SM4
++TargetedRHIs=PCD3D_SM5
++TargetedRHIs=PCD3D_SM4
++TargetedRHIs=GLSL_430
++TargetedRHIs=GLSL_150
+MinimumOSVersion=MSOS_Vista
+AudioDevice=
+AudioSampleRate=48000
+AudioCallbackBufferFrameSize=1024
+AudioNumBuffersToEnqueue=1
+AudioMaxChannels=0
+AudioNumSourceWorkers=4
+SpatializationPlugin=
+ReverbPlugin=
+OcclusionPlugin=
+
+
diff --git a/Templates/TP_nDisplayBP/Config/DefaultGame.ini b/Templates/TP_nDisplayBP/Config/DefaultGame.ini
new file mode 100644
--- /dev/null
+++ b/Templates/TP_nDisplayBP/Config/DefaultGame.ini
@@ -0,0 +1,7 @@
+[/Script/EngineSettings.GeneralProjectSettings]
+ProjectID=F20FEC0B4858196D9B8C90AD210974DE
+bUseBorderlessWindow=True
+
+[/Script/UnrealEd.ProjectPackagingSettings]
++DirectoriesToAlwaysStageAsNonUFS=(Path="ExampleConfigs")
+
diff --git a/Templates/TP_nDisplayBP/Config/DefaultInput.ini b/Templates/TP_nDisplayBP/Config/DefaultInput.ini
new file mode 100644
--- /dev/null
+++ b/Templates/TP_nDisplayBP/Config/DefaultInput.ini
@@ -0,0 +1,61 @@
+
+[/Script/Engine.InputSettings]
+-AxisConfig=(AxisKeyName="Gamepad_LeftX",AxisProperties=(DeadZone=0.25,Exponent=1.f,Sensitivity=1.f))
+-AxisConfig=(AxisKeyName="Gamepad_LeftY",AxisProperties=(DeadZone=0.25,Exponent=1.f,Sensitivity=1.f))
+-AxisConfig=(AxisKeyName="Gamepad_RightX",AxisProperties=(DeadZone=0.25,Exponent=1.f,Sensitivity=1.f))
+-AxisConfig=(AxisKeyName="Gamepad_RightY",AxisProperties=(DeadZone=0.25,Exponent=1.f,Sensitivity=1.f))
+-AxisConfig=(AxisKeyName="MouseX",AxisProperties=(DeadZone=0.f,Exponent=1.f,Sensitivity=0.07f))
+-AxisConfig=(AxisKeyName="MouseY",AxisProperties=(DeadZone=0.f,Exponent=1.f,Sensitivity=0.07f))
++AxisConfig=(AxisKeyName="Gamepad_LeftX",AxisProperties=(DeadZone=0.250000,Sensitivity=1.000000,Exponent=1.000000,bInvert=False))
++AxisConfig=(AxisKeyName="Gamepad_LeftY",AxisProperties=(DeadZone=0.250000,Sensitivity=1.000000,Exponent=1.000000,bInvert=False))
++AxisConfig=(AxisKeyName="Gamepad_RightX",AxisProperties=(DeadZone=0.250000,Sensitivity=1.000000,Exponent=1.000000,bInvert=False))
++AxisConfig=(AxisKeyName="Gamepad_RightY",AxisProperties=(DeadZone=0.250000,Sensitivity=1.000000,Exponent=1.000000,bInvert=False))
++AxisConfig=(AxisKeyName="MouseX",AxisProperties=(DeadZone=0.000000,Sensitivity=0.070000,Exponent=1.000000,bInvert=False))
++AxisConfig=(AxisKeyName="MouseY",AxisProperties=(DeadZone=0.000000,Sensitivity=0.070000,Exponent=1.000000,bInvert=False))
++AxisConfig=(AxisKeyName="MouseWheelAxis",AxisProperties=(DeadZone=0.000000,Sensitivity=1.000000,Exponent=1.000000,bInvert=False))
++AxisConfig=(AxisKeyName="Gamepad_LeftTriggerAxis",AxisProperties=(DeadZone=0.000000,Sensitivity=1.000000,Exponent=1.000000,bInvert=False))
++AxisConfig=(AxisKeyName="Gamepad_RightTriggerAxis",AxisProperties=(DeadZone=0.000000,Sensitivity=1.000000,Exponent=1.000000,bInvert=False))
++AxisConfig=(AxisKeyName="MotionController_Left_Thumbstick_X",AxisProperties=(DeadZone=0.000000,Sensitivity=1.000000,Exponent=1.000000,bInvert=False))
++AxisConfig=(AxisKeyName="MotionController_Left_Thumbstick_Y",AxisProperties=(DeadZone=0.000000,Sensitivity=1.000000,Exponent=1.000000,bInvert=False))
++AxisConfig=(AxisKeyName="MotionController_Left_TriggerAxis",AxisProperties=(DeadZone=0.000000,Sensitivity=1.000000,Exponent=1.000000,bInvert=False))
++AxisConfig=(AxisKeyName="MotionController_Left_Grip1Axis",AxisProperties=(DeadZone=0.000000,Sensitivity=1.000000,Exponent=1.000000,bInvert=False))
++AxisConfig=(AxisKeyName="MotionController_Left_Grip2Axis",AxisProperties=(DeadZone=0.000000,Sensitivity=1.000000,Exponent=1.000000,bInvert=False))
++AxisConfig=(AxisKeyName="MotionController_Right_Thumbstick_X",AxisProperties=(DeadZone=0.000000,Sensitivity=1.000000,Exponent=1.000000,bInvert=False))
++AxisConfig=(AxisKeyName="MotionController_Right_Thumbstick_Y",AxisProperties=(DeadZone=0.000000,Sensitivity=1.000000,Exponent=1.000000,bInvert=False))
++AxisConfig=(AxisKeyName="MotionController_Right_TriggerAxis",AxisProperties=(DeadZone=0.000000,Sensitivity=1.000000,Exponent=1.000000,bInvert=False))
++AxisConfig=(AxisKeyName="MotionController_Right_Grip1Axis",AxisProperties=(DeadZone=0.000000,Sensitivity=1.000000,Exponent=1.000000,bInvert=False))
++AxisConfig=(AxisKeyName="MotionController_Right_Grip2Axis",AxisProperties=(DeadZone=0.000000,Sensitivity=1.000000,Exponent=1.000000,bInvert=False))
++AxisConfig=(AxisKeyName="Gamepad_Special_Left_X",AxisProperties=(DeadZone=0.000000,Sensitivity=1.000000,Exponent=1.000000,bInvert=False))
++AxisConfig=(AxisKeyName="Gamepad_Special_Left_Y",AxisProperties=(DeadZone=0.000000,Sensitivity=1.000000,Exponent=1.000000,bInvert=False))
+bAltEnterTogglesFullscreen=True
+bF11TogglesFullscreen=True
+bUseMouseForTouch=False
+bEnableMouseSmoothing=True
+bEnableFOVScaling=True
+bCaptureMouseOnLaunch=True
+bDefaultViewportMouseLock=False
+bAlwaysShowTouchInterface=False
+bShowConsoleOnFourFingerTap=True
+bEnableGestureRecognizer=False
+DefaultViewportMouseCaptureMode=CapturePermanently_IncludingInitialMouseDown
+DefaultViewportMouseLockMode=LockOnCapture
+FOVScale=0.011110
+DoubleClickTime=0.200000
++AxisMappings=(AxisName="MoveForward",Scale=1.000000,Key=W)
++AxisMappings=(AxisName="MoveRight",Scale=1.000000,Key=D)
++AxisMappings=(AxisName="MoveUp",Scale=1.000000,Key=E)
++AxisMappings=(AxisName="TurnRate",Scale=-1.000000,Key=Left)
++AxisMappings=(AxisName="LookUpRate",Scale=-1.000000,Key=Up)
++AxisMappings=(AxisName="MoveForward",Scale=-1.000000,Key=S)
++AxisMappings=(AxisName="MoveRight",Scale=-1.000000,Key=A)
++AxisMappings=(AxisName="MoveUp",Scale=-1.000000,Key=Q)
++AxisMappings=(AxisName="TurnRate",Scale=1.000000,Key=Right)
++AxisMappings=(AxisName="TurnRate",Scale=1.000000,Key=MouseX)
++AxisMappings=(AxisName="LookUpRate",Scale=1.000000,Key=Down)
++AxisMappings=(AxisName="LookUpRate",Scale=-1.000000,Key=MouseY)
+DefaultTouchInterface=/Engine/MobileResources/HUD/DefaultVirtualJoysticks.DefaultVirtualJoysticks
+ConsoleKey=None
+-ConsoleKeys=Tilde
++ConsoleKeys=Tilde
+
+
diff --git a/Templates/TP_nDisplayBP/Config/TemplateDefs.ini b/Templates/TP_nDisplayBP/Config/TemplateDefs.ini
new file mode 100644
--- /dev/null
+++ b/Templates/TP_nDisplayBP/Config/TemplateDefs.ini
@@ -0,0 +1,28 @@
+﻿[/Script/GameProjectGeneration.TemplateProjectDefs]
+LocalizedDisplayNames=(Language="en",Text="nDisplay")
+LocalizedDescriptions=(Language="en",Text="This project template provides multi-display functionality using a cluster of PCs. Use as a starting point for stereoscropic/monoscopic Caves and PowerWall setups.")
+
+ClassTypes=GameMode, Actor
+AssetTypes=Static Mesh, Texture
+
+FoldersToIgnore=Binaries
+FoldersToIgnore=Build
+FoldersToIgnore=Intermediate
+FoldersToIgnore=Saved
+FoldersToIgnore=Media
+
+FilesToIgnore="%TEMPLATENAME%.uproject"
+FilesToIgnore="%TEMPLATENAME%.png"
+FilesToIgnore="Config/TemplateDefs.ini"
+FilesToIgnore="Config/config.ini"
+
+FolderRenames=(From="Source/%TEMPLATENAME%",To="Source/%PROJECTNAME%")
+FolderRenames=(From="Source/%TEMPLATENAME%Editor",To="Source/%PROJECTNAME%Editor")
+
+FilenameReplacements=(Extensions=("cpp","h","ini","cs"),From="%TEMPLATENAME_UPPERCASE%",To="%PROJECTNAME_UPPERCASE%",bCaseSensitive=true)
+FilenameReplacements=(Extensions=("cpp","h","ini","cs"),From="%TEMPLATENAME_LOWERCASE%",To="%PROJECTNAME_LOWERCASE%",bCaseSensitive=true)
+FilenameReplacements=(Extensions=("cpp","h","ini","cs"),From="%TEMPLATENAME%",To="%PROJECTNAME%",bCaseSensitive=false)
+
+ReplacementsInFiles=(Extensions=("cpp","h","ini","cs"),From="%TEMPLATENAME_UPPERCASE%",To="%PROJECTNAME_UPPERCASE%",bCaseSensitive=true)
+ReplacementsInFiles=(Extensions=("cpp","h","ini","cs"),From="%TEMPLATENAME_LOWERCASE%",To="%PROJECTNAME_LOWERCASE%",bCaseSensitive=true)
+ReplacementsInFiles=(Extensions=("cpp","h","ini","cs"),From="%TEMPLATENAME%",To="%PROJECTNAME%",bCaseSensitive=false)
\ No newline at end of file
diff --git a/Templates/TP_nDisplayBP/TP_nDisplayBP.uproject b/Templates/TP_nDisplayBP/TP_nDisplayBP.uproject
new file mode 100644
--- /dev/null
+++ b/Templates/TP_nDisplayBP/TP_nDisplayBP.uproject
@@ -0,0 +1,15 @@
+{
+	"FileVersion": 3,
+	"EngineAssociation": "",
+	"Category": "",
+	"Description": "",
+	"Plugins": [
+		{
+			"Name": "nDisplay",
+			"Enabled": true
+		}
+	],
+	"TargetPlatforms": [
+		"WindowsNoEditor"
+	]
+}
\ No newline at end of file
