Fix nDisplay template default rendering options to disable unsupported features

diff --git a/Templates/TP_nDisplayBP/Config/DefaultEngine.ini b/Templates/TP_nDisplayBP/Config/DefaultEngine.ini
--- a/Templates/TP_nDisplayBP/Config/DefaultEngine.ini
+++ b/Templates/TP_nDisplayBP/Config/DefaultEngine.ini
@@ -77,4 +77,7 @@ SpatializationPlugin=
 ReverbPlugin=
 OcclusionPlugin=
 
-
+[/Script/Engine.RendererSettings]
+r.DefaultFeature.Bloom=False
+r.DefaultFeature.AutoExposure=False
+r.DefaultFeature.MotionBlur=False
