Fixed warning due to outdated "VR" setting in template.

diff --git a/Templates/TP_CollaborativeBP/Config/TemplateDefs.ini b/Templates/TP_CollaborativeBP/Config/TemplateDefs.ini
--- a/Templates/TP_CollaborativeBP/Config/TemplateDefs.ini
+++ b/Templates/TP_CollaborativeBP/Config/TemplateDefs.ini
@@ -76,7 +76,7 @@ ReplacementsInFiles=(Extensions=("cpp","h","ini","cs"), From="%TEMPLATENAME%", T
 ; HardwareTarget - Choose between Mobile and Desktop hardware targets.
 ; GraphicsPreset - Choose between Maximum and Scalable graphics.
 ; StarterContent - Choose whether to add starter content to the project.
-; VR - Choose whether to enable Virtual Reality support for the project.
+; XR - Choose whether to enable Virtual Reality support for the project.
 ;HiddenSettings=All
 HiddenSettings=Languages
 HiddenSettings=HardwareTarget
diff --git a/Templates/TP_ME_BlankBP/Config/TemplateDefs.ini b/Templates/TP_ME_BlankBP/Config/TemplateDefs.ini
--- a/Templates/TP_ME_BlankBP/Config/TemplateDefs.ini
+++ b/Templates/TP_ME_BlankBP/Config/TemplateDefs.ini
@@ -75,7 +75,7 @@ ReplacementsInFiles=(Extensions=("cpp","h","ini","cs"), From="%TEMPLATENAME%", T
 ; HardwareTarget - Choose between Mobile and Desktop hardware targets.
 ; GraphicsPreset - Choose between Maximum and Scalable graphics.
 ; StarterContent - Choose whether to add starter content to the project.
-; VR - Choose whether to enable Virtual Reality support for the project.
+; XR - Choose whether to enable Virtual Reality support for the project.
 ;HiddenSettings=All
 HiddenSettings=Languages
 HiddenSettings=HardwareTarget
diff --git a/Templates/TP_ME_VProdBP/Config/TemplateDefs.ini b/Templates/TP_ME_VProdBP/Config/TemplateDefs.ini
--- a/Templates/TP_ME_VProdBP/Config/TemplateDefs.ini
+++ b/Templates/TP_ME_VProdBP/Config/TemplateDefs.ini
@@ -75,7 +75,7 @@ ReplacementsInFiles=(Extensions=("cpp","h","ini","cs"), From="%TEMPLATENAME%", T
 ; HardwareTarget - Choose between Mobile and Desktop hardware targets.
 ; GraphicsPreset - Choose between Maximum and Scalable graphics.
 ; StarterContent - Choose whether to add starter content to the project.
-; VR - Choose whether to enable Virtual Reality support for the project.
+; XR - Choose whether to enable Virtual Reality support for the project.
 ;HiddenSettings=All
 HiddenSettings=Languages
 HiddenSettings=HardwareTarget
diff --git a/Templates/TP_PhotoStudioBP/Config/TemplateDefs.ini b/Templates/TP_PhotoStudioBP/Config/TemplateDefs.ini
--- a/Templates/TP_PhotoStudioBP/Config/TemplateDefs.ini
+++ b/Templates/TP_PhotoStudioBP/Config/TemplateDefs.ini
@@ -75,7 +75,7 @@ ReplacementsInFiles=(Extensions=("cpp","h","ini","cs"), From="%TEMPLATENAME%", T
 ; HardwareTarget - Choose between Mobile and Desktop hardware targets.
 ; GraphicsPreset - Choose between Maximum and Scalable graphics.
 ; StarterContent - Choose whether to add starter content to the project.
-; VR - Choose whether to enable Virtual Reality support for the project.
+; XR - Choose whether to enable Virtual Reality support for the project.
 ;HiddenSettings=All
 HiddenSettings=Languages
 HiddenSettings=HardwareTarget
diff --git a/Templates/TP_nDisplayBP/Config/TemplateDefs.ini b/Templates/TP_nDisplayBP/Config/TemplateDefs.ini
--- a/Templates/TP_nDisplayBP/Config/TemplateDefs.ini
+++ b/Templates/TP_nDisplayBP/Config/TemplateDefs.ini
@@ -46,10 +46,10 @@ ReplacementsInFiles=(Extensions=("cpp","h","ini","cs"),From="%TEMPLATENAME%",To=
 ; HardwareTarget - Choose between Mobile and Desktop hardware targets.
 ; GraphicsPreset - Choose between Maximum and Scalable graphics.
 ; StarterContent - Choose whether to add starter content to the project.
-; VR - Choose whether to enable Virtual Reality support for the project.
+; XR - Choose whether to enable Virtual Reality support for the project.
 ;HiddenSettings=All
 HiddenSettings=Languages
 HiddenSettings=HardwareTarget
 HiddenSettings=GraphicsPreset 
 ;HiddenSettings=StarterContent 
-HiddenSettings=VR 
\ No newline at end of file
+HiddenSettings=XR 
\ No newline at end of file
