nDisplay: Linux support (Core)

diff --git a/Templates/TP_nDisplayBP/Config/Linux/LinuxEngine.ini b/Templates/TP_nDisplayBP/Config/Linux/LinuxEngine.ini
new file mode 100644
--- /dev/null
+++ b/Templates/TP_nDisplayBP/Config/Linux/LinuxEngine.ini
@@ -0,0 +1,10 @@
+[/Script/Engine.Engine]
+GameEngine=/Script/DisplayCluster.DisplayClusterGameEngine
+UnrealEdEngine=/Script/DisplayClusterEditor.DisplayClusterEditorEngine
+GameViewportClientClassName=/Script/DisplayCluster.DisplayClusterViewportClient
+
+[/Script/Engine.RendererSettings]
+r.DefaultBackBufferPixelFormat = 1
+
+[/Script/EngineSettings.GeneralProjectSettings]
+ bUseBorderlessWindow=True
