Revive nDisplay Template, nDisplay Template UE5.0 parity update, resaved all configs.
Allow CPU access to plane_1x1.uasset

diff --git a/Templates/TP_nDisplayBP/Config/DefaultEngine.ini b/Templates/TP_nDisplayBP/Config/DefaultEngine.ini
new file mode 100644
--- /dev/null
+++ b/Templates/TP_nDisplayBP/Config/DefaultEngine.ini
@@ -0,0 +1,29 @@
+[/Script/HardwareTargeting.HardwareTargetingSettings]
+TargetedHardwareClass=Desktop
+AppliedTargetedHardwareClass=Desktop
+DefaultGraphicsPerformance=Maximum
+AppliedDefaultGraphicsPerformance=Maximum
+
+[/Script/EngineSettings.GameMapsSettings]
+EditorStartupMap=/Game/Maps/Main.Main
+GameDefaultMap=/Game/Maps/Main.Main
+
+[/Script/DisplayClusterEditor.DisplayClusterEditorSettings]
+bEnabled=True
+
+[/Script/Engine.RendererSettings]
+r.DefaultFeature.Bloom=False
+r.DefaultFeature.AutoExposure=False
+r.DefaultFeature.MotionBlur=False
+r.SceneRenderTargetResizeMethod=2
+vr.AllowMotionBlurInVR=True
+r.DefaultFeature.AutoExposure.Method=1
+r.DefaultFeature.AutoExposure.Bias=1.000000
+
+[/Script/WindowsTargetPlatform.WindowsTargetSettings]
+DefaultGraphicsRHI=DefaultGraphicsRHI_DX11
++TargetedRHIs=PCD3D_SM5
+-TargetedRHIs=PCD3D_SM4
+-TargetedRHIs=GLSL_430
+-TargetedRHIs=GLSL_150
+
diff --git a/Templates/TP_nDisplayBP/Config/DefaultGame.ini b/Templates/TP_nDisplayBP/Config/DefaultGame.ini
new file mode 100644
--- /dev/null
+++ b/Templates/TP_nDisplayBP/Config/DefaultGame.ini
@@ -0,0 +1,7 @@
+[/Script/EngineSettings.GeneralProjectSettings]
+ProjectID=F20FEC0B4858196D9B8C90AD210974DE
+bUseBorderlessWindow=True
+
+[/Script/UnrealEd.ProjectPackagingSettings]
++DirectoriesToAlwaysStageAsNonUFS=(Path="ExampleConfigs")
+
diff --git a/Templates/TP_nDisplayBP/Config/DefaultInput.ini b/Templates/TP_nDisplayBP/Config/DefaultInput.ini
new file mode 100644
--- /dev/null
+++ b/Templates/TP_nDisplayBP/Config/DefaultInput.ini
@@ -0,0 +1,61 @@
+
+[/Script/Engine.InputSettings]
+-AxisConfig=(AxisKeyName="Gamepad_LeftX",AxisProperties=(DeadZone=0.25,Exponent=1.f,Sensitivity=1.f))
+-AxisConfig=(AxisKeyName="Gamepad_LeftY",AxisProperties=(DeadZone=0.25,Exponent=1.f,Sensitivity=1.f))
+-AxisConfig=(AxisKeyName="Gamepad_RightX",AxisProperties=(DeadZone=0.25,Exponent=1.f,Sensitivity=1.f))
+-AxisConfig=(AxisKeyName="Gamepad_RightY",AxisProperties=(DeadZone=0.25,Exponent=1.f,Sensitivity=1.f))
+-AxisConfig=(AxisKeyName="MouseX",AxisProperties=(DeadZone=0.f,Exponent=1.f,Sensitivity=0.07f))
+-AxisConfig=(AxisKeyName="MouseY",AxisProperties=(DeadZone=0.f,Exponent=1.f,Sensitivity=0.07f))
++AxisConfig=(AxisKeyName="Gamepad_LeftX",AxisProperties=(DeadZone=0.250000,Sensitivity=1.000000,Exponent=1.000000,bInvert=False))
++AxisConfig=(AxisKeyName="Gamepad_LeftY",AxisProperties=(DeadZone=0.250000,Sensitivity=1.000000,Exponent=1.000000,bInvert=False))
++AxisConfig=(AxisKeyName="Gamepad_RightX",AxisProperties=(DeadZone=0.250000,Sensitivity=1.000000,Exponent=1.000000,bInvert=False))
++AxisConfig=(AxisKeyName="Gamepad_RightY",AxisProperties=(DeadZone=0.250000,Sensitivity=1.000000,Exponent=1.000000,bInvert=False))
++AxisConfig=(AxisKeyName="MouseX",AxisProperties=(DeadZone=0.000000,Sensitivity=0.070000,Exponent=1.000000,bInvert=False))
++AxisConfig=(AxisKeyName="MouseY",AxisProperties=(DeadZone=0.000000,Sensitivity=0.070000,Exponent=1.000000,bInvert=False))
++AxisConfig=(AxisKeyName="MouseWheelAxis",AxisProperties=(DeadZone=0.000000,Sensitivity=1.000000,Exponent=1.000000,bInvert=False))
++AxisConfig=(AxisKeyName="Gamepad_LeftTriggerAxis",AxisProperties=(DeadZone=0.000000,Sensitivity=1.000000,Exponent=1.000000,bInvert=False))
++AxisConfig=(AxisKeyName="Gamepad_RightTriggerAxis",AxisProperties=(DeadZone=0.000000,Sensitivity=1.000000,Exponent=1.000000,bInvert=False))
++AxisConfig=(AxisKeyName="MotionController_Left_Thumbstick_X",AxisProperties=(DeadZone=0.000000,Sensitivity=1.000000,Exponent=1.000000,bInvert=False))
++AxisConfig=(AxisKeyName="MotionController_Left_Thumbstick_Y",AxisProperties=(DeadZone=0.000000,Sensitivity=1.000000,Exponent=1.000000,bInvert=False))
++AxisConfig=(AxisKeyName="MotionController_Left_TriggerAxis",AxisProperties=(DeadZone=0.000000,Sensitivity=1.000000,Exponent=1.000000,bInvert=False))
++AxisConfig=(AxisKeyName="MotionController_Left_Grip1Axis",AxisProperties=(DeadZone=0.000000,Sensitivity=1.000000,Exponent=1.000000,bInvert=False))
++AxisConfig=(AxisKeyName="MotionController_Left_Grip2Axis",AxisProperties=(DeadZone=0.000000,Sensitivity=1.000000,Exponent=1.000000,bInvert=False))
++AxisConfig=(AxisKeyName="MotionController_Right_Thumbstick_X",AxisProperties=(DeadZone=0.000000,Sensitivity=1.000000,Exponent=1.000000,bInvert=False))
++AxisConfig=(AxisKeyName="MotionController_Right_Thumbstick_Y",AxisProperties=(DeadZone=0.000000,Sensitivity=1.000000,Exponent=1.000000,bInvert=False))
++AxisConfig=(AxisKeyName="MotionController_Right_TriggerAxis",AxisProperties=(DeadZone=0.000000,Sensitivity=1.000000,Exponent=1.000000,bInvert=False))
++AxisConfig=(AxisKeyName="MotionController_Right_Grip1Axis",AxisProperties=(DeadZone=0.000000,Sensitivity=1.000000,Exponent=1.000000,bInvert=False))
++AxisConfig=(AxisKeyName="MotionController_Right_Grip2Axis",AxisProperties=(DeadZone=0.000000,Sensitivity=1.000000,Exponent=1.000000,bInvert=False))
++AxisConfig=(AxisKeyName="Gamepad_Special_Left_X",AxisProperties=(DeadZone=0.000000,Sensitivity=1.000000,Exponent=1.000000,bInvert=False))
++AxisConfig=(AxisKeyName="Gamepad_Special_Left_Y",AxisProperties=(DeadZone=0.000000,Sensitivity=1.000000,Exponent=1.000000,bInvert=False))
+bAltEnterTogglesFullscreen=True
+bF11TogglesFullscreen=True
+bUseMouseForTouch=False
+bEnableMouseSmoothing=True
+bEnableFOVScaling=True
+bCaptureMouseOnLaunch=True
+bDefaultViewportMouseLock=False
+bAlwaysShowTouchInterface=False
+bShowConsoleOnFourFingerTap=True
+bEnableGestureRecognizer=False
+DefaultViewportMouseCaptureMode=CapturePermanently_IncludingInitialMouseDown
+DefaultViewportMouseLockMode=LockOnCapture
+FOVScale=0.011110
+DoubleClickTime=0.200000
++AxisMappings=(AxisName="MoveForward",Scale=1.000000,Key=W)
++AxisMappings=(AxisName="MoveRight",Scale=1.000000,Key=D)
++AxisMappings=(AxisName="MoveUp",Scale=1.000000,Key=E)
++AxisMappings=(AxisName="TurnRate",Scale=-1.000000,Key=Left)
++AxisMappings=(AxisName="LookUpRate",Scale=-1.000000,Key=Up)
++AxisMappings=(AxisName="MoveForward",Scale=-1.000000,Key=S)
++AxisMappings=(AxisName="MoveRight",Scale=-1.000000,Key=A)
++AxisMappings=(AxisName="MoveUp",Scale=-1.000000,Key=Q)
++AxisMappings=(AxisName="TurnRate",Scale=1.000000,Key=Right)
++AxisMappings=(AxisName="TurnRate",Scale=1.000000,Key=MouseX)
++AxisMappings=(AxisName="LookUpRate",Scale=1.000000,Key=Down)
++AxisMappings=(AxisName="LookUpRate",Scale=-1.000000,Key=MouseY)
+DefaultTouchInterface=/Engine/MobileResources/HUD/DefaultVirtualJoysticks.DefaultVirtualJoysticks
+ConsoleKey=None
+-ConsoleKeys=Tilde
++ConsoleKeys=Tilde
+
+
diff --git a/Templates/TP_nDisplayBP/Config/TemplateDefs.ini b/Templates/TP_nDisplayBP/Config/TemplateDefs.ini
new file mode 100644
--- /dev/null
+++ b/Templates/TP_nDisplayBP/Config/TemplateDefs.ini
@@ -0,0 +1,55 @@
+﻿[/Script/GameProjectGeneration.TemplateProjectDefs]
+
+Categories=ME
+
+LocalizedDisplayNames=(Language="en",Text="nDisplay")
+LocalizedDescriptions=(Language="en",Text="This project template provides multi-display functionality using a cluster of PCs. Use as a starting point for stereoscropic/monoscopic Caves and PowerWall setups.")
+
+LocalizedDisplayNames=(Language="ko",Text="nDisplay")
+LocalizedDescriptions=(Language="ko",Text="PC 클러스터에 연결된 멀티 디스플레이 기능을 제공하는 프로젝트 템플릿입니다. 스테레오/모노스코픽 동굴이나 파워 월 구성을 위한 시작점으로 사용할 수 있습니다.")
+
+LocalizedDisplayNames=(Language="ja",Text="nDisplay")
+LocalizedDescriptions=(Language="ja",Text="PCクラスタを使ったマルチディスプレイ機能のテンプレートです。立体視/平面視の Cave や PowerWall のセットアップでの土台として使用してください。")
+
+LocalizedDisplayNames=(Language="zh-Hans",Text="nDisplay")
+LocalizedDescriptions=(Language="zh-Hans",Text="此项目模板提供了使用PC集群的多显示设备功能。它可作为立体/平面Caves以及PowerWall设置的基础。")
+
+ClassTypes=GameMode, Actor
+AssetTypes=Static Mesh, Texture
+
+FoldersToIgnore=Binaries
+FoldersToIgnore=Build
+FoldersToIgnore=Intermediate
+FoldersToIgnore=Saved
+FoldersToIgnore=Media
+
+FilesToIgnore="%TEMPLATENAME%.uproject"
+FilesToIgnore="%TEMPLATENAME%.png"
+FilesToIgnore="Config/TemplateDefs.ini"
+FilesToIgnore="Config/config.ini"
+
+FolderRenames=(From="Source/%TEMPLATENAME%",To="Source/%PROJECTNAME%")
+FolderRenames=(From="Source/%TEMPLATENAME%Editor",To="Source/%PROJECTNAME%Editor")
+
+FilenameReplacements=(Extensions=("cpp","h","ini","cs"),From="%TEMPLATENAME_UPPERCASE%",To="%PROJECTNAME_UPPERCASE%",bCaseSensitive=true)
+FilenameReplacements=(Extensions=("cpp","h","ini","cs"),From="%TEMPLATENAME_LOWERCASE%",To="%PROJECTNAME_LOWERCASE%",bCaseSensitive=true)
+FilenameReplacements=(Extensions=("cpp","h","ini","cs"),From="%TEMPLATENAME%",To="%PROJECTNAME%",bCaseSensitive=false)
+
+ReplacementsInFiles=(Extensions=("cpp","h","ini","cs"),From="%TEMPLATENAME_UPPERCASE%",To="%PROJECTNAME_UPPERCASE%",bCaseSensitive=true)
+ReplacementsInFiles=(Extensions=("cpp","h","ini","cs"),From="%TEMPLATENAME_LOWERCASE%",To="%PROJECTNAME_LOWERCASE%",bCaseSensitive=true)
+ReplacementsInFiles=(Extensions=("cpp","h","ini","cs"),From="%TEMPLATENAME%",To="%PROJECTNAME%",bCaseSensitive=false)
+
+; Which settings on the project settings page should be hidden for this template. Default: none
+; Allowed values are:
+; All - Hides the settings page entirely
+; Languages - Choose between C++ and Blueprint versions.
+; HardwareTarget - Choose between Mobile and Desktop hardware targets.
+; GraphicsPreset - Choose between Maximum and Scalable graphics.
+; StarterContent - Choose whether to add starter content to the project.
+; XR - Choose whether to enable Virtual Reality support for the project.
+;HiddenSettings=All
+HiddenSettings=Languages
+HiddenSettings=HardwareTarget
+HiddenSettings=GraphicsPreset 
+;HiddenSettings=StarterContent 
+HiddenSettings=XR 
\ No newline at end of file
diff --git a/Templates/TP_nDisplayBP/Config/Windows/WindowsEngine.ini b/Templates/TP_nDisplayBP/Config/Windows/WindowsEngine.ini
new file mode 100644
--- /dev/null
+++ b/Templates/TP_nDisplayBP/Config/Windows/WindowsEngine.ini
@@ -0,0 +1,6 @@
+[/Script/Engine.Engine]
+GameEngine=/Script/DisplayCluster.DisplayClusterGameEngine
+UnrealEdEngine=/Script/DisplayClusterEditor.DisplayClusterEditorEngine
+GameViewportClientClassName=/Script/DisplayCluster.DisplayClusterViewportClient
+
+
diff --git a/Templates/TP_nDisplayBP/TP_nDisplayBP.uproject b/Templates/TP_nDisplayBP/TP_nDisplayBP.uproject
new file mode 100644
--- /dev/null
+++ b/Templates/TP_nDisplayBP/TP_nDisplayBP.uproject
@@ -0,0 +1,37 @@
+{
+	"FileVersion": 3,
+	"EngineAssociation": "",
+	"Category": "",
+	"Description": "",
+	"Plugins": [
+		{
+			"Name": "nDisplay",
+			"Enabled": true,
+			"SupportedTargetPlatforms": [
+				"Win64",
+				"Linux"
+			]
+		},
+		{
+			"Name": "OculusVR",
+			"Enabled": false,
+			"SupportedTargetPlatforms": [
+				"Win64"
+			]
+		},
+		{
+			"Name": "SteamVR",
+			"Enabled": false,
+			"SupportedTargetPlatforms": [
+				"Win64"
+			]
+		},
+		{
+			"Name": "Switchboard",
+			"Enabled": true
+		}
+	],
+	"TargetPlatforms": [
+		"WindowsNoEditor"
+	]
+}
\ No newline at end of file
