Templates: Fixed some ME templates not showing up in Linux

Also fixed a couple of unversioned uassets.

TP_DMXBP
---------
InstalledEngineFilters.xml
    Removed it from CopyEditorFilterWin64
    Added it to CopyEditorFilter
    Remove it from ProjectsToBuildDDCWin64
    Added it to ProjectsToBuildDDC
InstalledBuild.cs
    Moved out of win64 condition in GetDdcProjects

With this, it should now also show up in Linux and Mac.

TP_nDisplayBP
-------------

.uproject
    Removed Vscouting and OpenXR dependency.
    Added Linux in its TargetPlatforms
InstalledEngineFilters.xml
    Added to CopyEditorFilterLinux (so we'll have Win+Linux)

With this, it should now also show up in Linux.

#virtualized

diff --git a/Engine/Build/InstalledEngineFilters.xml b/Engine/Build/InstalledEngineFilters.xml
--- a/Engine/Build/InstalledEngineFilters.xml
+++ b/Engine/Build/InstalledEngineFilters.xml
@@ -64,12 +64,12 @@
 		Templates/TP_VehicleAdvBP/TP_VehicleAdvBP.uproject
 		Templates/TP_VirtualRealityBP/TP_VirtualRealityBP.uproject
 		Templates/TP_SIM_BlankBP/TP_SIM_BlankBP.uproject
+		Templates/TP_DMXBP/TP_DMXBP.uproject
 		Samples/StarterContent/StarterContent.uproject
 	</Property>
 	<!-- Win64 specific DDC -->
 	<Property Name="ProjectsToBuildDDCWin64" Value="$(ProjectsToBuildDDC)">
 		Templates/TP_InCamVFXBP/TP_InCamVFXBP.uproject
-		Templates/TP_DMXBP/TP_DMXBP.uproject
 	</Property>
 	
 	<!-- Define Editor Filters -->
@@ -135,6 +135,8 @@
 		Templates/TP_SIM_Blank/...
 		Templates/TP_SIM_BlankBP/...
 
+		Templates/TP_DMXBP/...
+
 		<!-- Shared template resources -->
 		Templates/TemplateResources/...
 
@@ -297,9 +299,6 @@
 		Templates/TP_nDisplayBP/...
 		Templates/TP_InCamVFXBP/...
 
-		<!-- Enterprise Templates Win64 Only -->
-		Templates/TP_DMXBP/...
-
 		<!-- DotNet -->
 		Engine/Binaries/ThirdParty/DotNet/6.0.302/windows/...
 		Engine/Binaries/ThirdParty/VisualStudio/...
@@ -413,6 +412,9 @@
 		Engine/Binaries/Linux/dump_syms
 		Engine/Binaries/Linux/BreakpadSymbolEncoder
 
+		<!-- Virtual Production Templates -->
+		Templates/TP_nDisplayBP/...
+
 		<!-- DotNet -->
 		Engine/Binaries/ThirdParty/DotNet/6.0.302/linux/...
 		Engine/Binaries/ThirdParty/VisualStudio/...
diff --git a/Engine/Source/Programs/AutomationTool/BuildGraph/Scripts/InstalledBuild.cs b/Engine/Source/Programs/AutomationTool/BuildGraph/Scripts/InstalledBuild.cs
--- a/Engine/Source/Programs/AutomationTool/BuildGraph/Scripts/InstalledBuild.cs
+++ b/Engine/Source/Programs/AutomationTool/BuildGraph/Scripts/InstalledBuild.cs
@@ -71,12 +71,12 @@ namespace AutomationTool
 			Projects.Add("Templates/TP_TopDownBP/TP_TopDownBP.uproject");
 			Projects.Add("Templates/TP_VehicleAdvBP/TP_VehicleAdvBP.uproject");
 			Projects.Add("Templates/TP_VirtualRealityBP/TP_VirtualRealityBP.uproject");
-//			Projects.Add("Samples/StarterContent/StarterContent.uproject");
+			Projects.Add("Templates/TP_DMXBP/TP_DMXBP.uproject");
+			//			Projects.Add("Samples/StarterContent/StarterContent.uproject");
 
 			if (Platform == UnrealTargetPlatform.Win64)
 			{
 				Projects.Add("Templates/TP_InCamVFXBP/TP_InCamVFXBP.uproject");
-				Projects.Add("Templates/TP_DMXBP/TP_DMXBP.uproject");
 			}
 			return Projects;
 		}
diff --git a/Templates/TP_nDisplayBP/Config/DefaultInput.ini b/Templates/TP_nDisplayBP/Config/DefaultInput.ini
--- a/Templates/TP_nDisplayBP/Config/DefaultInput.ini
+++ b/Templates/TP_nDisplayBP/Config/DefaultInput.ini
@@ -1,4 +1,5 @@
 
+
 [/Script/Engine.InputSettings]
 -AxisConfig=(AxisKeyName="Gamepad_LeftX",AxisProperties=(DeadZone=0.25,Exponent=1.f,Sensitivity=1.f))
 -AxisConfig=(AxisKeyName="Gamepad_LeftY",AxisProperties=(DeadZone=0.25,Exponent=1.f,Sensitivity=1.f))
@@ -6,6 +7,8 @@
 -AxisConfig=(AxisKeyName="Gamepad_RightY",AxisProperties=(DeadZone=0.25,Exponent=1.f,Sensitivity=1.f))
 -AxisConfig=(AxisKeyName="MouseX",AxisProperties=(DeadZone=0.f,Exponent=1.f,Sensitivity=0.07f))
 -AxisConfig=(AxisKeyName="MouseY",AxisProperties=(DeadZone=0.f,Exponent=1.f,Sensitivity=0.07f))
+-AxisConfig=(AxisKeyName="Mouse2D",AxisProperties=(DeadZone=0.f,Exponent=1.f,Sensitivity=0.07f))
++AxisConfig=(AxisKeyName="Mouse2D",AxisProperties=(DeadZone=0.000000,Sensitivity=0.070000,Exponent=1.000000,bInvert=False))
 +AxisConfig=(AxisKeyName="Gamepad_LeftX",AxisProperties=(DeadZone=0.250000,Sensitivity=1.000000,Exponent=1.000000,bInvert=False))
 +AxisConfig=(AxisKeyName="Gamepad_LeftY",AxisProperties=(DeadZone=0.250000,Sensitivity=1.000000,Exponent=1.000000,bInvert=False))
 +AxisConfig=(AxisKeyName="Gamepad_RightX",AxisProperties=(DeadZone=0.250000,Sensitivity=1.000000,Exponent=1.000000,bInvert=False))
@@ -27,16 +30,62 @@
 +AxisConfig=(AxisKeyName="MotionController_Right_Grip2Axis",AxisProperties=(DeadZone=0.000000,Sensitivity=1.000000,Exponent=1.000000,bInvert=False))
 +AxisConfig=(AxisKeyName="Gamepad_Special_Left_X",AxisProperties=(DeadZone=0.000000,Sensitivity=1.000000,Exponent=1.000000,bInvert=False))
 +AxisConfig=(AxisKeyName="Gamepad_Special_Left_Y",AxisProperties=(DeadZone=0.000000,Sensitivity=1.000000,Exponent=1.000000,bInvert=False))
++AxisConfig=(AxisKeyName="Vive_Left_Trigger_Axis",AxisProperties=(DeadZone=0.000000,Sensitivity=1.000000,Exponent=1.000000,bInvert=False))
++AxisConfig=(AxisKeyName="Vive_Left_Trackpad_X",AxisProperties=(DeadZone=0.000000,Sensitivity=1.000000,Exponent=1.000000,bInvert=False))
++AxisConfig=(AxisKeyName="Vive_Left_Trackpad_Y",AxisProperties=(DeadZone=0.000000,Sensitivity=1.000000,Exponent=1.000000,bInvert=False))
++AxisConfig=(AxisKeyName="Vive_Right_Trigger_Axis",AxisProperties=(DeadZone=0.000000,Sensitivity=1.000000,Exponent=1.000000,bInvert=False))
++AxisConfig=(AxisKeyName="Vive_Right_Trackpad_X",AxisProperties=(DeadZone=0.000000,Sensitivity=1.000000,Exponent=1.000000,bInvert=False))
++AxisConfig=(AxisKeyName="Vive_Right_Trackpad_Y",AxisProperties=(DeadZone=0.000000,Sensitivity=1.000000,Exponent=1.000000,bInvert=False))
++AxisConfig=(AxisKeyName="MixedReality_Left_Trigger_Axis",AxisProperties=(DeadZone=0.000000,Sensitivity=1.000000,Exponent=1.000000,bInvert=False))
++AxisConfig=(AxisKeyName="MixedReality_Left_Thumbstick_X",AxisProperties=(DeadZone=0.000000,Sensitivity=1.000000,Exponent=1.000000,bInvert=False))
++AxisConfig=(AxisKeyName="MixedReality_Left_Thumbstick_Y",AxisProperties=(DeadZone=0.000000,Sensitivity=1.000000,Exponent=1.000000,bInvert=False))
++AxisConfig=(AxisKeyName="MixedReality_Left_Trackpad_X",AxisProperties=(DeadZone=0.000000,Sensitivity=1.000000,Exponent=1.000000,bInvert=False))
++AxisConfig=(AxisKeyName="MixedReality_Left_Trackpad_Y",AxisProperties=(DeadZone=0.000000,Sensitivity=1.000000,Exponent=1.000000,bInvert=False))
++AxisConfig=(AxisKeyName="MixedReality_Right_Trigger_Axis",AxisProperties=(DeadZone=0.000000,Sensitivity=1.000000,Exponent=1.000000,bInvert=False))
++AxisConfig=(AxisKeyName="MixedReality_Right_Thumbstick_X",AxisProperties=(DeadZone=0.000000,Sensitivity=1.000000,Exponent=1.000000,bInvert=False))
++AxisConfig=(AxisKeyName="MixedReality_Right_Thumbstick_Y",AxisProperties=(DeadZone=0.000000,Sensitivity=1.000000,Exponent=1.000000,bInvert=False))
++AxisConfig=(AxisKeyName="MixedReality_Right_Trackpad_X",AxisProperties=(DeadZone=0.000000,Sensitivity=1.000000,Exponent=1.000000,bInvert=False))
++AxisConfig=(AxisKeyName="MixedReality_Right_Trackpad_Y",AxisProperties=(DeadZone=0.000000,Sensitivity=1.000000,Exponent=1.000000,bInvert=False))
++AxisConfig=(AxisKeyName="OculusTouch_Left_Grip_Axis",AxisProperties=(DeadZone=0.000000,Sensitivity=1.000000,Exponent=1.000000,bInvert=False))
++AxisConfig=(AxisKeyName="OculusTouch_Left_Trigger_Axis",AxisProperties=(DeadZone=0.000000,Sensitivity=1.000000,Exponent=1.000000,bInvert=False))
++AxisConfig=(AxisKeyName="OculusTouch_Left_Thumbstick_X",AxisProperties=(DeadZone=0.000000,Sensitivity=1.000000,Exponent=1.000000,bInvert=False))
++AxisConfig=(AxisKeyName="OculusTouch_Left_Thumbstick_Y",AxisProperties=(DeadZone=0.000000,Sensitivity=1.000000,Exponent=1.000000,bInvert=False))
++AxisConfig=(AxisKeyName="OculusTouch_Right_Grip_Axis",AxisProperties=(DeadZone=0.000000,Sensitivity=1.000000,Exponent=1.000000,bInvert=False))
++AxisConfig=(AxisKeyName="OculusTouch_Right_Trigger_Axis",AxisProperties=(DeadZone=0.000000,Sensitivity=1.000000,Exponent=1.000000,bInvert=False))
++AxisConfig=(AxisKeyName="OculusTouch_Right_Thumbstick_X",AxisProperties=(DeadZone=0.000000,Sensitivity=1.000000,Exponent=1.000000,bInvert=False))
++AxisConfig=(AxisKeyName="OculusTouch_Right_Thumbstick_Y",AxisProperties=(DeadZone=0.000000,Sensitivity=1.000000,Exponent=1.000000,bInvert=False))
++AxisConfig=(AxisKeyName="ValveIndex_Left_Grip_Axis",AxisProperties=(DeadZone=0.000000,Sensitivity=1.000000,Exponent=1.000000,bInvert=False))
++AxisConfig=(AxisKeyName="ValveIndex_Left_Grip_Force",AxisProperties=(DeadZone=0.000000,Sensitivity=1.000000,Exponent=1.000000,bInvert=False))
++AxisConfig=(AxisKeyName="ValveIndex_Left_Trigger_Axis",AxisProperties=(DeadZone=0.000000,Sensitivity=1.000000,Exponent=1.000000,bInvert=False))
++AxisConfig=(AxisKeyName="ValveIndex_Left_Thumbstick_X",AxisProperties=(DeadZone=0.000000,Sensitivity=1.000000,Exponent=1.000000,bInvert=False))
++AxisConfig=(AxisKeyName="ValveIndex_Left_Thumbstick_Y",AxisProperties=(DeadZone=0.000000,Sensitivity=1.000000,Exponent=1.000000,bInvert=False))
++AxisConfig=(AxisKeyName="ValveIndex_Left_Trackpad_X",AxisProperties=(DeadZone=0.000000,Sensitivity=1.000000,Exponent=1.000000,bInvert=False))
++AxisConfig=(AxisKeyName="ValveIndex_Left_Trackpad_Y",AxisProperties=(DeadZone=0.000000,Sensitivity=1.000000,Exponent=1.000000,bInvert=False))
++AxisConfig=(AxisKeyName="ValveIndex_Left_Trackpad_Force",AxisProperties=(DeadZone=0.000000,Sensitivity=1.000000,Exponent=1.000000,bInvert=False))
++AxisConfig=(AxisKeyName="ValveIndex_Right_Grip_Axis",AxisProperties=(DeadZone=0.000000,Sensitivity=1.000000,Exponent=1.000000,bInvert=False))
++AxisConfig=(AxisKeyName="ValveIndex_Right_Grip_Force",AxisProperties=(DeadZone=0.000000,Sensitivity=1.000000,Exponent=1.000000,bInvert=False))
++AxisConfig=(AxisKeyName="ValveIndex_Right_Trigger_Axis",AxisProperties=(DeadZone=0.000000,Sensitivity=1.000000,Exponent=1.000000,bInvert=False))
++AxisConfig=(AxisKeyName="ValveIndex_Right_Thumbstick_X",AxisProperties=(DeadZone=0.000000,Sensitivity=1.000000,Exponent=1.000000,bInvert=False))
++AxisConfig=(AxisKeyName="ValveIndex_Right_Thumbstick_Y",AxisProperties=(DeadZone=0.000000,Sensitivity=1.000000,Exponent=1.000000,bInvert=False))
++AxisConfig=(AxisKeyName="ValveIndex_Right_Trackpad_X",AxisProperties=(DeadZone=0.000000,Sensitivity=1.000000,Exponent=1.000000,bInvert=False))
++AxisConfig=(AxisKeyName="ValveIndex_Right_Trackpad_Y",AxisProperties=(DeadZone=0.000000,Sensitivity=1.000000,Exponent=1.000000,bInvert=False))
++AxisConfig=(AxisKeyName="ValveIndex_Right_Trackpad_Force",AxisProperties=(DeadZone=0.000000,Sensitivity=1.000000,Exponent=1.000000,bInvert=False))
 bAltEnterTogglesFullscreen=True
 bF11TogglesFullscreen=True
 bUseMouseForTouch=False
 bEnableMouseSmoothing=True
 bEnableFOVScaling=True
 bCaptureMouseOnLaunch=True
-bDefaultViewportMouseLock=False
+bEnableLegacyInputScales=True
+bEnableMotionControls=True
+bFilterInputByPlatformUser=False
+bEnableInputDeviceSubsystem=True
+bShouldFlushPressedKeysOnViewportFocusLost=True
+bEnableDynamicComponentInputBinding=True
 bAlwaysShowTouchInterface=False
 bShowConsoleOnFourFingerTap=True
 bEnableGestureRecognizer=False
+bUseAutocorrect=False
 DefaultViewportMouseCaptureMode=CapturePermanently_IncludingInitialMouseDown
 DefaultViewportMouseLockMode=LockOnCapture
 FOVScale=0.011110
@@ -54,9 +103,8 @@ DoubleClickTime=0.200000
 +AxisMappings=(AxisName="LookUpRate",Scale=1.000000,Key=Down)
 +AxisMappings=(AxisName="LookUpRate",Scale=-1.000000,Key=MouseY)
 DefaultPlayerInputClass=/Script/DisplayCluster.DisplayClusterPlayerInput
+DefaultInputComponentClass=/Script/EnhancedInput.EnhancedInputComponent
 DefaultTouchInterface=/Engine/MobileResources/HUD/DefaultVirtualJoysticks.DefaultVirtualJoysticks
-ConsoleKey=None
 -ConsoleKeys=Tilde
 +ConsoleKeys=Tilde
 
-
diff --git a/Templates/TP_nDisplayBP/TP_nDisplayBP.uproject b/Templates/TP_nDisplayBP/TP_nDisplayBP.uproject
--- a/Templates/TP_nDisplayBP/TP_nDisplayBP.uproject
+++ b/Templates/TP_nDisplayBP/TP_nDisplayBP.uproject
@@ -12,30 +12,13 @@
 				"Linux"
 			]
 		},
-		{
-			"Name": "OpenXR",
-			"Enabled": true,
-			"SupportedTargetPlatforms": [
-				"Win64",
-				"Linux",
-				"Android"
-			]
-		},
-		{
-			"Name": "VirtualScouting",
-			"Enabled": true,
-			"SupportedTargetPlatforms": [
-				"Win64",
-				"Linux"
-			]
-		},
 		{
 			"Name": "Switchboard",
 			"Enabled": true
 		}
 	],
 	"TargetPlatforms": [
-		"WindowsNoEditor",
-		"Windows"
+		"Windows",
+		"Linux"
 	]
 }
\ No newline at end of file
